Experimental: OSC Mod Mixing Plugin

diff --git a/Engine/Plugins/Runtime/OSCModulationMixing/OSCModulationMixing.uplugin b/Engine/Plugins/Runtime/OSCModulationMixing/OSCModulationMixing.uplugin
new file mode 100644
--- /dev/null
+++ b/Engine/Plugins/Runtime/OSCModulationMixing/OSCModulationMixing.uplugin
@@ -0,0 +1,36 @@
+{
+	"FileVersion": 3,
+	"Version": 1,
+	"VersionName": "1.0",
+	"FriendlyName": "OSCModulationMixing",
+	"Description": "Utility Blueprint objects and functions enabling Control Modulation mix profiling using the OSC (Open Sound Control) protocol.",
+	"Category": "Audio",
+	"CreatedBy": "Epic Games, Inc.",
+	"CreatedByURL": "http://epicgames.com",
+	"DocsURL": "",
+	"MarketplaceURL": "",
+	"SupportURL": "",
+	"CanContainContent": true,
+	"IsBetaVersion": false,
+	"IsExperimentalVersion": true,
+	"Installed": false,
+	"Plugins":
+	[
+		{
+			"Name": "AudioModulation",
+			"Enabled": true
+		},
+		{
+			"Name": "OSC",
+			"Enabled": true
+		}
+	],
+	"Modules":
+	[
+		{
+			"Name": "OSCModulationMixing",
+			"Type": "Runtime",
+			"LoadingPhase": "PreLoadingScreen"
+		}
+	]
+}
\ No newline at end of file
diff --git a/Engine/Plugins/Runtime/OSCModulationMixing/Source/OSCModulationMixing/OSCModulationMixing.Build.cs b/Engine/Plugins/Runtime/OSCModulationMixing/Source/OSCModulationMixing/OSCModulationMixing.Build.cs
new file mode 100644
--- /dev/null
+++ b/Engine/Plugins/Runtime/OSCModulationMixing/Source/OSCModulationMixing/OSCModulationMixing.Build.cs
@@ -0,0 +1,30 @@
+// Copyright Epic Games, Inc. All Rights Reserved.
+
+using UnrealBuildTool;
+
+public class OSCModulationMixing : ModuleRules
+{
+	public OSCModulationMixing(ReadOnlyTargetRules Target) : base(Target)
+	{
+		OptimizeCode = CodeOptimization.Never;
+
+		PCHUsage = ModuleRules.PCHUsageMode.UseExplicitOrSharedPCHs;
+
+		PublicDependencyModuleNames.AddRange(
+			new string[]
+			{
+				"Core",
+				"AudioModulation",
+				"OSC",
+			}
+		);
+
+		PrivateDependencyModuleNames.AddRange(
+			new string[]
+			{
+				"CoreUObject",
+				"Engine",
+			}
+		);
+	}
+}
diff --git a/Engine/Plugins/Runtime/OSCModulationMixing/Source/OSCModulationMixing/Private/OSCModulationMixing.cpp b/Engine/Plugins/Runtime/OSCModulationMixing/Source/OSCModulationMixing/Private/OSCModulationMixing.cpp
new file mode 100644
--- /dev/null
+++ b/Engine/Plugins/Runtime/OSCModulationMixing/Source/OSCModulationMixing/Private/OSCModulationMixing.cpp
@@ -0,0 +1,17 @@
+// Copyright Epic Games, Inc. All Rights Reserved.
+
+#include "OSCModulationMixing.h"
+
+#define LOCTEXT_NAMESPACE "FOSCModulationMixingModule"
+
+void FOSCModulationMixingModule::StartupModule()
+{
+}
+
+void FOSCModulationMixingModule::ShutdownModule()
+{
+}
+
+#undef LOCTEXT_NAMESPACE
+	
+IMPLEMENT_MODULE(FOSCModulationMixingModule, OSCModulationMixing)
\ No newline at end of file
diff --git a/Engine/Plugins/Runtime/OSCModulationMixing/Source/OSCModulationMixing/Private/OSCModulationMixingStatics.cpp b/Engine/Plugins/Runtime/OSCModulationMixing/Source/OSCModulationMixing/Private/OSCModulationMixingStatics.cpp
new file mode 100644
--- /dev/null
+++ b/Engine/Plugins/Runtime/OSCModulationMixing/Source/OSCModulationMixing/Private/OSCModulationMixingStatics.cpp
@@ -0,0 +1,149 @@
+// Copyright Epic Games, Inc. All Rights Reserved.
+
+#include "OSCModulationMixingStatics.h"
+
+#include "OSCClient.h"
+#include "OSCManager.h"
+#include "OSCModulationMixing.h"
+#include "OSCServer.h"
+
+#include "HAL/IConsoleManager.h"
+#include "SoundControlBusMix.h"
+#include "UObject/WeakObjectPtrTemplates.h"
+
+
+namespace OSCModulation
+{
+	namespace Addresses
+	{
+		static const FOSCAddress MixLoad	 = FString(TEXT("/Mix/Load"));
+		static const FOSCAddress ProfileLoad = FString(TEXT("/Mix/Profile/Load"));
+		static const FOSCAddress ProfileSave = FString(TEXT("/Mix/Profile/Save"));
+	} // namespace Addresses
+} // namespace OSCModulation
+
+UOSCModulationMixingStatics::UOSCModulationMixingStatics(const FObjectInitializer& ObjectInitializer)
+: Super(ObjectInitializer)
+{
+}
+
+void UOSCModulationMixingStatics::CopyChannelsToOSCBundle(UObject* WorldContextObject, const FOSCAddress& InPathAddress, const TArray<FSoundControlBusMixChannel>& InChannels, FOSCBundle& OutBundle)
+{
+	FOSCMessage RequestMessage;
+	RequestMessage.SetAddress(OSCModulation::Addresses::MixLoad);
+	UOSCManager::AddMessageToBundle(RequestMessage, OutBundle);
+
+	FOSCMessage Message;
+	Message.SetAddress(InPathAddress);
+	UOSCManager::AddMessageToBundle(Message, OutBundle);
+
+	for (const FSoundControlBusMixChannel& Channel : InChannels)
+	{
+		if (Channel.Bus)
+		{
+			FOSCMessage ChannelMessage;
+			ChannelMessage.SetAddress(UOSCManager::OSCAddressFromObjectPath(Channel.Bus));
+
+			UOSCManager::AddFloat(ChannelMessage, Channel.Value.AttackTime);
+			UOSCManager::AddFloat(ChannelMessage, Channel.Value.ReleaseTime);
+			UOSCManager::AddFloat(ChannelMessage, Channel.Value.TargetValue);
+
+			UClass* BusClass = Channel.Bus->GetClass();
+			const FString ClassName = BusClass ? BusClass->GetName() : FString();
+			UOSCManager::AddString(ChannelMessage, ClassName);
+
+			UOSCManager::AddMessageToBundle(ChannelMessage, OutBundle);
+		}
+	}
+}
+
+void UOSCModulationMixingStatics::CopyMixToOSCBundle(UObject* WorldContextObject, USoundControlBusMix* InMix, UPARAM(ref) FOSCBundle& OutBundle)
+{
+	if (!InMix)
+	{
+		return;
+	}
+
+	CopyChannelsToOSCBundle(WorldContextObject, UOSCManager::OSCAddressFromObjectPath(InMix), InMix->Channels, OutBundle);
+}
+
+FOSCAddress UOSCModulationMixingStatics::GetProfileLoadPath()
+{
+	return OSCModulation::Addresses::ProfileLoad;
+}
+
+FOSCAddress UOSCModulationMixingStatics::GetProfileSavePath()
+{
+	return OSCModulation::Addresses::ProfileSave;
+}
+
+FOSCAddress UOSCModulationMixingStatics::GetMixLoadPattern()
+{
+	return FOSCAddress(OSCModulation::Addresses::MixLoad.GetFullPath() / TEXT("*"));
+}
+
+EOSCModulationBundle UOSCModulationMixingStatics::GetOSCBundleType(const FOSCBundle& InBundle)
+{
+	TArray<FOSCMessage> Messages = UOSCManager::GetMessagesFromBundle(InBundle);
+	if (Messages.Num() == 0)
+	{
+		return EOSCModulationBundle::Invalid;
+	}
+
+	FOSCAddress BundleAddress = UOSCManager::GetOSCMessageAddress(Messages[0]);
+	if (BundleAddress == OSCModulation::Addresses::MixLoad)
+	{
+		return EOSCModulationBundle::LoadMix;
+	}
+	static_assert(static_cast<int32>(EOSCModulationBundle::Count) == 2, "Possible missing bundle case coverage");
+	// Add additional bundle types here
+
+	return EOSCModulationBundle::Invalid;
+}
+
+void UOSCModulationMixingStatics::RequestMix(UObject* WorldContextObject, UOSCClient* InClient, const FOSCAddress& InMixPath)
+{
+	if (InClient)
+	{
+		FOSCMessage Message;
+
+		const FOSCAddress MixAddress = UOSCManager::OSCAddressFromObjectPathString(InMixPath.GetFullPath());
+		Message.SetAddress(OSCModulation::Addresses::MixLoad / MixAddress);
+
+		InClient->SendOSCMessage(Message);
+	}
+}
+
+TArray<FSoundModulationValue> UOSCModulationMixingStatics::OSCBundleToChannelValues(UObject* WorldContextObject, const FOSCBundle& InBundle, FOSCAddress& OutMixPath, TArray<FOSCAddress>& OutBusPaths, TArray<FString>& OutBusClassNames)
+{
+	TArray<FSoundModulationValue> ChannelArray;
+	OutBusPaths.Reset();
+
+	const TArray<FOSCMessage> Messages = UOSCManager::GetMessagesFromBundle(InBundle);
+	if (Messages.Num() > 1)
+	{
+		if (Messages[0].GetAddress() == OSCModulation::Addresses::MixLoad)
+		{
+			OutMixPath = Messages[1].GetAddress();
+
+			for (int32 i = 2; i < Messages.Num(); ++i)
+			{
+				FSoundModulationValue Value;
+				UOSCManager::GetFloat(Messages[i], 0, Value.AttackTime);
+				UOSCManager::GetFloat(Messages[i], 1, Value.ReleaseTime);
+				UOSCManager::GetFloat(Messages[i], 2, Value.TargetValue);
+
+				FString BusClass;
+				UOSCManager::GetString(Messages[i], 3, BusClass);
+				OutBusClassNames.Add(BusClass);
+
+				OutBusPaths.Add(UOSCManager::GetOSCMessageAddress(Messages[i]));
+				ChannelArray.Add(Value);
+			}
+			return ChannelArray;
+		}
+	}
+
+	OutMixPath = FOSCAddress();
+	return ChannelArray;
+}
diff --git a/Engine/Plugins/Runtime/OSCModulationMixing/Source/OSCModulationMixing/Public/OSCModulationMixing.h b/Engine/Plugins/Runtime/OSCModulationMixing/Source/OSCModulationMixing/Public/OSCModulationMixing.h
new file mode 100644
--- /dev/null
+++ b/Engine/Plugins/Runtime/OSCModulationMixing/Source/OSCModulationMixing/Public/OSCModulationMixing.h
@@ -0,0 +1,14 @@
+// Copyright Epic Games, Inc. All Rights Reserved.
+
+#pragma once
+
+#include "Modules/ModuleManager.h"
+
+class FOSCModulationMixingModule : public IModuleInterface
+{
+public:
+
+	/** IModuleInterface implementation */
+	virtual void StartupModule() override;
+	virtual void ShutdownModule() override;
+};
diff --git a/Engine/Plugins/Runtime/OSCModulationMixing/Source/OSCModulationMixing/Public/OSCModulationMixingStatics.h b/Engine/Plugins/Runtime/OSCModulationMixing/Source/OSCModulationMixing/Public/OSCModulationMixingStatics.h
new file mode 100644
--- /dev/null
+++ b/Engine/Plugins/Runtime/OSCModulationMixing/Source/OSCModulationMixing/Public/OSCModulationMixingStatics.h
@@ -0,0 +1,71 @@
+// Copyright Epic Games, Inc. All Rights Reserved.
+#pragma once
+
+#include "Kismet/BlueprintFunctionLibrary.h"
+
+#include "OSCAddress.h"
+#include "OSCBundle.h"
+#include "OSCClient.h"
+#include "OSCMessage.h"
+#include "SoundControlBusMix.h"
+#include "SoundModulationValue.h"
+
+#include "OSCModulationMixingStatics.generated.h"
+
+// Forward Declarations
+class USoundControlBusMix;
+
+UENUM()
+enum class EOSCModulationBundle : uint8
+{
+	Invalid,
+	LoadMix,
+	Count UMETA(Hidden)
+};
+
+UENUM()
+enum class EOSCModulationMessage : uint8
+{
+	Invalid,
+	LoadProfile,
+	SaveProfile,
+	Count UMETA(Hidden)
+};
+
+UCLASS()
+class OSCMODULATIONMIXING_API UOSCModulationMixingStatics : public UBlueprintFunctionLibrary
+{
+	GENERATED_UCLASS_BODY()
+	
+	/** Returns OSC Address pattern for loading a mix */
+	UFUNCTION(BlueprintCallable, Category = "Audio|OSC|Modulation", DisplayName = "Get Load Mix Pattern")
+	static FOSCAddress GetMixLoadPattern();
+
+	/** Returns OSC Address path for loading a profile */
+	UFUNCTION(BlueprintCallable, Category = "Audio|OSC|Modulation", DisplayName = "Get Load Profile Path")
+	static FOSCAddress GetProfileLoadPath();
+
+	/** Returns OSC Address path for saving a profile */
+	UFUNCTION(BlueprintCallable, Category = "Audio|OSC|Modulation", DisplayName = "Get Save Profile Path")
+	static FOSCAddress GetProfileSavePath();
+
+	/** Converts channel array to OSCBundle representation to send over network via OSC protocol */
+	UFUNCTION(BlueprintCallable, Category = "Audio|OSC|Modulation", DisplayName = "Copy Mix Channels to OSC Bundle", meta = (WorldContext = "WorldContextObject"))
+	static void CopyChannelsToOSCBundle(UObject* WorldContextObject, const FOSCAddress& PathAddress, const TArray<FSoundControlBusMixChannel>& Channels, UPARAM(ref) FOSCBundle& Bundle);
+
+	/** Converts Control Bus Mix to OSCBundle representation to send over network via OSC protocol */
+	UFUNCTION(BlueprintCallable, Category = "Audio|OSC|Modulation", DisplayName = "Copy Mix State to OSC Bundle", meta = (WorldContext = "WorldContextObject"))
+	static void CopyMixToOSCBundle(UObject* WorldContextObject, USoundControlBusMix* Mix, UPARAM(ref) FOSCBundle& Bundle);
+
+	/** Returns whether bundle contains recognized payload of OSC Modulation Data */
+	UFUNCTION(BlueprintCallable, Category = "Audio|OSC|Modulation", DisplayName = "Get OSC Modulation Bundle Type")
+	static EOSCModulationBundle GetOSCBundleType(const FOSCBundle& Bundle);
+
+	/** Request mix update from server with loaded content. */
+	UFUNCTION(BlueprintCallable, Category = "Audio|OSC|Modulation", DisplayName = "Request Mix State Update", meta = (WorldContext = "WorldContextObject"))
+	static void RequestMix(UObject* WorldContextObject, UOSCClient* Client, const FOSCAddress& MixPath);
+
+	/** Converts OSCBundle to Control Bus Values & Mix Path from which it came */
+	UFUNCTION(BlueprintCallable, Category = "Audio|OSC|Modulation", DisplayName = "OSCBundle To Channel Values", meta = (WorldContext = "WorldContextObject"))
+	static UPARAM(DisplayName = "Bus Values") TArray<FSoundModulationValue> OSCBundleToChannelValues(UObject* WorldContextObject, const FOSCBundle& Bundle, FOSCAddress& MixPath, TArray<FOSCAddress>& BusPaths, TArray<FString>& BusClassNames);
+};
