Updated OSCModulationMixing to not compile with optimizations off so it matches other modules that use Shared PCHs. This speeds up compile times and removes the need to have a 2.03 GB shared Engine PCH.

diff --git a/Engine/Plugins/Runtime/OSCModulationMixing/Source/OSCModulationMixing/OSCModulationMixing.Build.cs b/Engine/Plugins/Runtime/OSCModulationMixing/Source/OSCModulationMixing/OSCModulationMixing.Build.cs
--- a/Engine/Plugins/Runtime/OSCModulationMixing/Source/OSCModulationMixing/OSCModulationMixing.Build.cs
+++ b/Engine/Plugins/Runtime/OSCModulationMixing/Source/OSCModulationMixing/OSCModulationMixing.Build.cs
@@ -6,8 +6,6 @@ public class OSCModulationMixing : ModuleRules
 {
 	public OSCModulationMixing(ReadOnlyTargetRules Target) : base(Target)
 	{
-		OptimizeCode = CodeOptimization.Never;
-
 		PCHUsage = ModuleRules.PCHUsageMode.UseExplicitOrSharedPCHs;
 
 		PublicDependencyModuleNames.AddRange(
